Some more UX tweaks for Zen Staus Bar

diff --git a/Engine/Source/Editor/EditorStyle/Private/StarshipStyle.cpp b/Engine/Source/Editor/EditorStyle/Private/StarshipStyle.cpp
--- a/Engine/Source/Editor/EditorStyle/Private/StarshipStyle.cpp
+++ b/Engine/Source/Editor/EditorStyle/Private/StarshipStyle.cpp
@@ -1096,11 +1096,17 @@ void FStarshipEditorStyle::FStyle::SetupGeneralStyles()
 
 	// Zen
 	{
-		Set("Zen.Icons.LaunchDashboard", new CORE_IMAGE_BRUSH_SVG("Starship/Common/SessionFrontend", Icon16x16));
-		Set("Zen.Icons.Server", new CORE_IMAGE_BRUSH_SVG("Starship/Common/SessionFrontend", Icon16x16));
+		Set("Zen.Icons.Server", new CORE_IMAGE_BRUSH_SVG("Starship/Common/Server", Icon16x16));
+		Set("Zen.Icons.Server.Start", new CORE_IMAGE_BRUSH_SVG("Starship/Common/Play", Icon16x16));
+		Set("Zen.Icons.Server.Stop", new CORE_IMAGE_BRUSH_SVG("Starship/Common/Stop", Icon16x16));
+		Set("Zen.Icons.Server.Restart", new CORE_IMAGE_BRUSH_SVG("Starship/Common/Update", Icon16x16));
+		Set("Zen.Icons.LaunchDashboard", new CORE_IMAGE_BRUSH_SVG("Starship/Common/Monitor", Icon16x16));
+		Set("Zen.Icons.ImportSnapshot", new CORE_IMAGE_BRUSH_SVG("Starship/Common/Import", Icon16x16));
+		Set("Zen.Icons.Store", new CORE_IMAGE_BRUSH_SVG("Starship/Common/Cylinder", Icon16x16));
 		Set("Zen.Icons.WebBrowser", new IMAGE_BRUSH_SVG("Starship/Common/WebBrowser", Icon12x12));
 		Set("Zen.Icons.FolderExplore", new IMAGE_BRUSH_SVG("Starship/ContentBrowser/show-in-explorer", Icon12x12));
-		Set("Zen.Icons.Clipboard", new IMAGE_BRUSH_SVG("Starship/Common/Clipboard", Icon16x16));
+		Set("Zen.Icons.Clipboard", new IMAGE_BRUSH_SVG("Starship/Common/Clipboard", Icon16x16));	
+
 	}
 
 	// Derived Data
diff --git a/Engine/Source/Editor/ZenEditor/Private/SZenStatusBar.cpp b/Engine/Source/Editor/ZenEditor/Private/SZenStatusBar.cpp
--- a/Engine/Source/Editor/ZenEditor/Private/SZenStatusBar.cpp
+++ b/Engine/Source/Editor/ZenEditor/Private/SZenStatusBar.cpp
@@ -47,6 +47,8 @@
 #include "Widgets/SOverlay.h"
 #include "Widgets/Text/STextBlock.h"
 #include "ZenDashboardLauncher.h"
+#include "ZenServiceInstanceManager.h"
+#include "Experimental/ZenServerInterface.h"
 
 class SWidget;
 struct FSlateBrush;
@@ -55,6 +57,15 @@ struct FSlateBrush;
 
 
 TSharedRef<FUICommandList> FZenStausBarCommands::ActionList(new FUICommandList());
+TSharedPtr<UE::Zen::FServiceInstanceManager> FZenStausBarCommands::ServiceInstanceManager;
+
+static const FName ZenMenuName(TEXT("ZenStatusBarMenu"));
+static const FName ZenMenuServerSectionName(TEXT("ZenStatusBarMenu.ServerSection"));
+static const FName ZenMenuStoreSectionName(TEXT("ZenStatusBarMenu.StoreSection"));
+static const FName ZenMenuToolsSectionName(TEXT("ZenStatusBarMenu.ToolsSection"));
+static const FName ZenMenuCacheSectionName(TEXT("ZenStatusBarMenu.CacheSection"));
+static const FName ZenMenuServerActionsSubMenuName(TEXT("ZenStatusBarMenu.ServerActions"));
+
 
 FZenStausBarCommands::FZenStausBarCommands()
 	: TCommands<FZenStausBarCommands>
@@ -68,15 +79,21 @@ FZenStausBarCommands::FZenStausBarCommands()
 
 void FZenStausBarCommands::RegisterCommands()
 {
-	UI_COMMAND(ChangeSettings, "Change Cache Settings", "Opens a dialog to change Cache settings.", EUserInterfaceActionType::Button, FInputChord());
+	ServiceInstanceManager = MakeShared<UE::Zen::FServiceInstanceManager>();
+
+	UI_COMMAND(ChangeCacheSettings, "Cache Settings", "Opens a dialog to change Cache settings.", EUserInterfaceActionType::Button, FInputChord());
 	UI_COMMAND(ViewCacheStatistics, "Cache Statistics", "Opens the Cache Statistics panel.", EUserInterfaceActionType::Button, FInputChord());
 	UI_COMMAND(ViewResourceUsage, "Cache Resource Usage", "Opens the Cache Resource Usage panel.", EUserInterfaceActionType::Button, FInputChord());
 	UI_COMMAND(ViewServerStatus, "Server Status", "Opens the Zen Server Status panel.", EUserInterfaceActionType::Button, FInputChord());
 	UI_COMMAND(LaunchZenDashboard, "Launch Dashboard", "Launches the Zen Dashboard utility.", EUserInterfaceActionType::Button, FInputChord());
+	UI_COMMAND(ViewZenStore, "View Store", "Opens the Zen Store page in your browser.", EUserInterfaceActionType::Button, FInputChord());
+	UI_COMMAND(StartZenServer, "Start Server", "Starts the Zen Server.", EUserInterfaceActionType::Button, FInputChord());
+	UI_COMMAND(StopZenServer, "Stop Server", "Stops the Zen Server.", EUserInterfaceActionType::Button, FInputChord());
+	UI_COMMAND(RestartZenServer, "Restart Server", "Restarts the Zen Server.", EUserInterfaceActionType::Button, FInputChord());
 
 	ActionList->MapAction(
-		ChangeSettings,
-		FExecuteAction::CreateStatic(&FZenStausBarCommands::ChangeSettings_Clicked)
+		ChangeCacheSettings,
+		FExecuteAction::CreateStatic(&FZenStausBarCommands::ChangeCacheSettings_Clicked)
 	);
 
 	ActionList->MapAction(
@@ -98,9 +115,29 @@ void FZenStausBarCommands::RegisterCommands()
 		LaunchZenDashboard,
 		FExecuteAction::CreateStatic(&FZenStausBarCommands::LaunchZenDashboard_Clicked)
 	);
+
+	ActionList->MapAction(
+		ViewZenStore,
+		FExecuteAction::CreateStatic(&FZenStausBarCommands::ViewZenStore_Clicked)
+	);
+
+	ActionList->MapAction(
+		StartZenServer,
+		FExecuteAction::CreateStatic(&FZenStausBarCommands::StartZenServer_Clicked)
+	);
+
+	ActionList->MapAction(
+		StopZenServer,
+		FExecuteAction::CreateStatic(&FZenStausBarCommands::StopZenServer_Clicked)
+	);
+
+	ActionList->MapAction(
+		RestartZenServer,
+		FExecuteAction::CreateStatic(&FZenStausBarCommands::RestartZenServer_Clicked)
+	);
 }
 
-void FZenStausBarCommands::ChangeSettings_Clicked()
+void FZenStausBarCommands::ChangeCacheSettings_Clicked()
 {
 	FModuleManager::LoadModuleChecked<ISettingsModule>("Settings").ShowViewer("Editor", "General", "Global");
 }
@@ -129,12 +166,41 @@ void FZenStausBarCommands::LaunchZenDashboard_Clicked()
 	FZenDashboardLauncher::Get()->StartZenDashboard(FZenDashboardLauncher::Get()->GetZenDashboardApplicationPath());
 }
 
+void FZenStausBarCommands::StartZenServer_Clicked()
+{
+	FZenEditor& ZenEditor = FModuleManager::LoadModuleChecked<FZenEditor>("ZenEditor");
+	ZenEditor.StartZenServer();
+}
+
+void FZenStausBarCommands::StopZenServer_Clicked()
+{
+	FZenEditor& ZenEditor = FModuleManager::LoadModuleChecked<FZenEditor>("ZenEditor");
+	ZenEditor.StopZenServer();
+}
+
+void FZenStausBarCommands::RestartZenServer_Clicked()
+{
+	FZenEditor& ZenEditor = FModuleManager::LoadModuleChecked<FZenEditor>("ZenEditor");
+	ZenEditor.RestartZenServer();
+}
+
+void FZenStausBarCommands::ViewZenStore_Clicked()
+{
+	using namespace UE::Zen;
+	TSharedPtr<FZenServiceInstance> ServerInstance = ServiceInstanceManager.Get()->GetZenServiceInstance();
+
+	if (ServerInstance.IsValid())
+	{
+		FPlatformProcess::LaunchURL(*FString::Printf(TEXT("%s/dashboard/?"), ServerInstance->GetURL()), nullptr, nullptr);
+	}
+}
+
 TSharedRef<SWidget> SZenStatusBarWidget::CreateStatusBarMenu()
 {
-	UToolMenu* Menu = UToolMenus::Get()->RegisterMenu("StatusBar.ToolBar.Zen", NAME_None, EMultiBoxType::Menu, false);
+	UToolMenu* Menu = UToolMenus::Get()->RegisterMenu(ZenMenuName, NAME_None, EMultiBoxType::Menu, false);
 
 	{
-		FToolMenuSection& Section = Menu->AddSection("ZenMenuToolsSection", LOCTEXT("ZenMenuToolsSection", "Tools"));
+		FToolMenuSection& Section = Menu->AddSection(ZenMenuToolsSectionName, LOCTEXT("ZenMenuToolsSection", "Tools"));
 
 		Section.AddMenuEntry(
 			FZenStausBarCommands::Get().LaunchZenDashboard,
@@ -143,9 +209,8 @@ TSharedRef<SWidget> SZenStatusBarWidget::CreateStatusBarMenu()
 			FSlateIcon(FAppStyle::GetAppStyleSetName(), "Zen.Icons.LaunchDashboard")
 		);
 	}
-
 	{
-		FToolMenuSection& Section = Menu->AddSection("ZenMenuServerSection", LOCTEXT("ZenMenuServerSection", "Server"));
+		FToolMenuSection& Section = Menu->AddSection(ZenMenuServerSectionName, LOCTEXT("ZenMenuServerSection", "Server"));
 
 		Section.AddMenuEntry(
 			FZenStausBarCommands::Get().ViewServerStatus,
@@ -153,17 +218,57 @@ TSharedRef<SWidget> SZenStatusBarWidget::CreateStatusBarMenu()
 			TAttribute<FText>(),
 			FSlateIcon(FAppStyle::GetAppStyleSetName(), "Zen.Icons.Server")
 		);
+
+		{
+			Section.AddSubMenu(TEXT("ZenServerActionSubMenu"), LOCTEXT("ZenServerActions", "Server Actions"), LOCTEXT("ZenServerActionsSubMenu", "Server Actions"),
+				FNewToolMenuDelegate::CreateLambda([this](UToolMenu* InSubMenu)
+					{
+						using namespace UE::Zen;
+						TSharedPtr<FZenServiceInstance> ServerInstance = ServiceInstanceManager.Get()->GetZenServiceInstance();
+						FToolMenuSection& Section = InSubMenu->AddSection(NAME_None);
+						
+						if (ServerInstance.IsValid() && ServerInstance->IsServiceRunning())
+						{
+							Section.AddMenuEntry(
+								FZenStausBarCommands::Get().StopZenServer,
+								TAttribute<FText>(),
+								TAttribute<FText>(),
+								FSlateIcon(FAppStyle::GetAppStyleSetName(), "Zen.Icons.Server.Stop")
+							);
+
+							Section.AddMenuEntry(
+								FZenStausBarCommands::Get().RestartZenServer,
+								TAttribute<FText>(),
+								TAttribute<FText>(),
+								FSlateIcon(FAppStyle::GetAppStyleSetName(), "Zen.Icons.Server.Restart")
+							);
+						}
+						else
+						{
+							Section.AddMenuEntry(
+								FZenStausBarCommands::Get().StartZenServer,
+								TAttribute<FText>(),
+								TAttribute<FText>(),
+								FSlateIcon(FAppStyle::GetAppStyleSetName(), "Zen.Icons.Server.Start")
+							);
+						}
+					}));
+		}
 	}
 
 	{
-		FToolMenuSection& Section = Menu->AddSection("ZenMenuDerivedDataSection", LOCTEXT("ZenMenuDerivedDataSection", "Cache"));
+		FToolMenuSection& Section = Menu->AddSection(ZenMenuStoreSectionName, LOCTEXT("ZenMenuStoreSection", "Store"));
 
 		Section.AddMenuEntry(
-			FZenStausBarCommands::Get().ChangeSettings,
+			FZenStausBarCommands::Get().ViewZenStore,
 			TAttribute<FText>(),
 			TAttribute<FText>(),
-			FSlateIcon(FAppStyle::GetAppStyleSetName(), "DerivedData.Cache.Settings")
+			FSlateIcon(FAppStyle::GetAppStyleSetName(), "Zen.Icons.Store")
 		);
+	}
+
+	{
+		FToolMenuSection& Section = Menu->AddSection(ZenMenuCacheSectionName, LOCTEXT("ZenMenuDerivedDataSection", "Cache"));
 
 		Section.AddMenuEntry(
 			FZenStausBarCommands::Get().ViewCacheStatistics,
@@ -178,13 +283,22 @@ TSharedRef<SWidget> SZenStatusBarWidget::CreateStatusBarMenu()
 			TAttribute<FText>(),
 			FSlateIcon(FAppStyle::GetAppStyleSetName(), "DerivedData.ResourceUsage")
 		);
+
+		Section.AddMenuEntry(
+			FZenStausBarCommands::Get().ChangeCacheSettings,
+			TAttribute<FText>(),
+			TAttribute<FText>(),
+			FSlateIcon(FAppStyle::GetAppStyleSetName(), "DerivedData.Cache.Settings")
+		);
 	}
 
-	return UToolMenus::Get()->GenerateWidget("StatusBar.ToolBar.Zen", FToolMenuContext(FZenStausBarCommands::ActionList));
+	return UToolMenus::Get()->GenerateWidget(ZenMenuName, FToolMenuContext(FZenStausBarCommands::ActionList));
 }
 
 void SZenStatusBarWidget::Construct(const FArguments& InArgs)
 {	
+	ServiceInstanceManager = MakeShared<UE::Zen::FServiceInstanceManager>();
+
 	this->ChildSlot
 	[
 		SNew(SComboButton)
@@ -208,8 +322,8 @@ void SZenStatusBarWidget::Construct(const FArguments& InArgs)
 				[
 					SNew(SImage)
 					.ColorAndOpacity(FSlateColor::UseForeground())
-					.Image_Lambda([this] { return GetRemoteCacheStateBackgroundIcon();  })
-					.ToolTipText_Lambda([this] { return GetRemoteCacheToolTipText(); })
+					.Image_Lambda([this] { return GetServerStateBackgroundIcon();  })
+					.ToolTipText_Lambda([this] { return GetServerStateToolTipText(); })
 				]
 
 				+ SOverlay::Slot()
@@ -218,8 +332,8 @@ void SZenStatusBarWidget::Construct(const FArguments& InArgs)
 				[
 					SNew(SImage)
 					.ColorAndOpacity(FSlateColor::UseForeground())
-					.Image_Lambda([this] { return GetRemoteCacheStateBadgeIcon();  })
-					.ToolTipText_Lambda([this] { return GetRemoteCacheToolTipText(); })
+					.Image_Lambda([this] { return GetServerStateBadgeIcon();  })
+					.ToolTipText_Lambda([this] { return GetServerStateToolTipText(); })
 				]
 
 				+ SOverlay::Slot()
@@ -228,8 +342,8 @@ void SZenStatusBarWidget::Construct(const FArguments& InArgs)
 				[
 					SNew(SImage)
 					.Image(FAppStyle::Get().GetBrush("DerivedData.RemoteCache.Uploading"))
-					.ColorAndOpacity_Lambda([this] { return ( FDerivedDataInformation::IsUploading() && FDerivedDataInformation::GetRemoteCacheState() == ERemoteCacheState::Busy )? FLinearColor::White.CopyWithNewOpacity(FMath::MakePulsatingValue(ElapsedUploadTime, 2)) : FLinearColor(0,0,0,0); })
-					.ToolTipText_Lambda([this] { return GetRemoteCacheToolTipText(); })
+					.ColorAndOpacity_Lambda([this] { return IsUploading? FLinearColor::White.CopyWithNewOpacity(FMath::MakePulsatingValue(ElapsedUploadTime, 2)) : FLinearColor(0,0,0,0); })
+					.ToolTipText_Lambda([this] { return GetServerStateToolTipText(); })
 				]
 
 				+ SOverlay::Slot()
@@ -238,8 +352,8 @@ void SZenStatusBarWidget::Construct(const FArguments& InArgs)
 				[
 					SNew(SImage)
 					.Image(FAppStyle::Get().GetBrush("DerivedData.RemoteCache.Downloading"))
-					.ColorAndOpacity_Lambda([this] { return ( FDerivedDataInformation::IsDownloading() && FDerivedDataInformation::GetRemoteCacheState()==ERemoteCacheState::Busy ) ? FLinearColor::White.CopyWithNewOpacity(FMath::MakePulsatingValue(ElapsedDownloadTime, 2)) : FLinearColor(0, 0, 0, 0); })
-					.ToolTipText_Lambda([this] { return GetRemoteCacheToolTipText(); })
+					.ColorAndOpacity_Lambda([this] { return IsDownloading ? FLinearColor::White.CopyWithNewOpacity(FMath::MakePulsatingValue(ElapsedDownloadTime, 2)) : FLinearColor(0, 0, 0, 0); })
+					.ToolTipText_Lambda([this] { return GetServerStateToolTipText(); })
 				]				
 			]
 
@@ -262,11 +376,62 @@ void SZenStatusBarWidget::Construct(const FArguments& InArgs)
 
 EActiveTimerReturnType SZenStatusBarWidget::UpdateBusyIndicator(double InCurrentTime, float InDeltaTime)
 {	
+	using namespace UE::Zen;
+	TSharedPtr<FZenServiceInstance> ServerInstance = ServiceInstanceManager.Get()->GetZenServiceInstance();
+
+	static int64 LastTotalHits=0;
+	static int64 LastTotalReads=0;
+	static int64 LastTotalWrites=0;
+
+	if (ServerInstance.IsValid() && ServerInstance->IsServiceRunning())
+	{
+		FZenCacheStats CacheStats;
+		ServerInstance->GetCacheStats(CacheStats);
+
+		FZenProjectStats ProjectStats;
+		ServerInstance->GetProjectStats(ProjectStats);
+
+		const int64 TotalCacheHits = CacheStats.General.Hits + CacheStats.General.CidHits;
+		const int64 TotalCacheWrites = CacheStats.General.Writes + CacheStats.General.CidWrites;
+		const int64 TotalProjectReads = ProjectStats.General.Project.ReadCount + ProjectStats.General.Oplog.ReadCount;
+		const int64 TotalProjectWrites = ProjectStats.General.Project.ReadCount + ProjectStats.General.Oplog.ReadCount;
+
+		const int64 TotalHits = TotalCacheHits;
+		const int64 TotalReads = TotalProjectReads;
+		const int64 TotalWrites = TotalCacheWrites + TotalProjectWrites;
+			
+		IsRunning = true;
+		IsDownloading = CacheStats.Upstream.Reading;
+		IsUploading = CacheStats.Upstream.Writing;
+		IsReading = TotalReads > LastTotalReads;
+		IsWriting = TotalWrites > LastTotalWrites;
+
+		LastTotalHits = TotalHits;
+		LastTotalReads = TotalReads;
+		LastTotalWrites = TotalWrites;
+	}
+	else
+	{
+		IsRunning = false;
+		IsDownloading = false;
+		IsUploading = false;
+		IsReading = false;
+		IsWriting = false;
+
+		LastTotalHits = 0;
+		LastTotalReads = 0;
+		LastTotalWrites = 0;
+	}
+
+	IsBusy = IsUploading || IsDownloading || IsReading || IsWriting;
+
 	FDerivedDataInformation::UpdateRemoteCacheState();
 
-	bBusy = GetDerivedDataCache()->AnyAsyncRequestsRemaining();
+	IsBusy			= GetDerivedDataCache()->AnyAsyncRequestsRemaining();
+	IsUploading		= FDerivedDataInformation::IsUploading() && FDerivedDataInformation::GetRemoteCacheState() == ERemoteCacheState::Busy;
+	IsDownloading	= FDerivedDataInformation::IsDownloading() && FDerivedDataInformation::GetRemoteCacheState() == ERemoteCacheState::Busy;
 
-	if (FDerivedDataInformation::IsUploading())
+	if (IsUploading)
 	{
 		ElapsedUploadTime += fmod(InDeltaTime,3600.0);
 	}
@@ -275,7 +440,7 @@ EActiveTimerReturnType SZenStatusBarWidget::UpdateBusyIndicator(double InCurrent
 		ElapsedUploadTime = 0.0;
 	}
 
-	if (FDerivedDataInformation::IsDownloading())
+	if (IsDownloading)
 	{
 		ElapsedDownloadTime += fmod(InDeltaTime, 3600.0);
 	}
@@ -284,7 +449,7 @@ EActiveTimerReturnType SZenStatusBarWidget::UpdateBusyIndicator(double InCurrent
 		ElapsedDownloadTime = 0.0;
 	}
 
-	if (bBusy)
+	if (IsBusy)
 	{
 		ElapsedBusyTime += fmod(InDeltaTime, 3600.0);
 	}
@@ -348,11 +513,7 @@ EActiveTimerReturnType SZenStatusBarWidget::UpdateWarnings(double InCurrentTime,
 
 FText SZenStatusBarWidget::GetTitleToolTipText() const
 {
-	FTextBuilder DescBuilder;
-
-	DescBuilder.AppendLineFormat(LOCTEXT("GraphNameText", "Graph : {0}"), FText::FromString(GetDerivedDataCache()->GetGraphName()));
-	
-	return DescBuilder.ToText();
+	return GetServerStateToolTipText();
 }
 
 FText SZenStatusBarWidget::GetTitleText() const
@@ -360,80 +521,41 @@ FText SZenStatusBarWidget::GetTitleText() const
 	return LOCTEXT("ZenStatusBarName", "Zen Server");
 }
 
-FText SZenStatusBarWidget::GetRemoteCacheToolTipText() const
+FText SZenStatusBarWidget::GetServerStateToolTipText() const
 {
 	FTextBuilder DescBuilder;
-
-	if (FDerivedDataInformation::GetRemoteCacheState() == ERemoteCacheState::Warning)
+	
+	if (IsRunning)
 	{
-		DescBuilder.AppendLineFormat(LOCTEXT("RemoteCacheErrorText", "WARNING\t: {0}\n"), FDerivedDataInformation::GetRemoteCacheWarningMessage());
+		DescBuilder.AppendLineFormat(LOCTEXT("ZenServerStaus", "{0}"), FText::FromString(IsBusy ? TEXT("Busy") : TEXT("Idle")));
+	}
+	else
+	{
+		DescBuilder.AppendLineFormat(LOCTEXT("ZenServerStaus", "{0}"), FText::FromString(TEXT("Stopped")));
 	}
 
-	DescBuilder.AppendLine(LOCTEXT("RemoteCacheToolTipText", "Remote Cache\n"));
-	DescBuilder.AppendLineFormat(LOCTEXT("RemoteCacheConnectedText", "Connected\t: {0}"), FText::FromString((FDerivedDataInformation::GetHasRemoteCache() ? TEXT("Yes") : TEXT("No"))));
-	DescBuilder.AppendLineFormat(LOCTEXT("RemoteCacheStatusText", "Status\t: {0}"), FDerivedDataInformation::GetRemoteCacheStateAsText() );
-	
-	const double DownloadedBytesMB = FUnitConversion::Convert(FDerivedDataInformation::GetCacheActivitySizeBytes(true, false), EUnit::Bytes, EUnit::Megabytes);
-	const double UploadedBytesMB = FUnitConversion::Convert(FDerivedDataInformation::GetCacheActivitySizeBytes(false, false), EUnit::Bytes, EUnit::Megabytes);
-
-	DescBuilder.AppendLineFormat(LOCTEXT("RemoteCacheDownloaded", "Downloaded\t: {0} MB"), DownloadedBytesMB);
-	DescBuilder.AppendLineFormat(LOCTEXT("RemoteCacheUploaded",	"Uploaded\t: {0} MB"), UploadedBytesMB);
-	
 	return DescBuilder.ToText();
 }
 
-const FSlateBrush* SZenStatusBarWidget::GetRemoteCacheStateBackgroundIcon() const
+const FSlateBrush* SZenStatusBarWidget::GetServerStateBackgroundIcon() const
 {
-	switch ( FDerivedDataInformation::GetRemoteCacheState())
+	if (IsRunning)
 	{
-		case ERemoteCacheState::Idle :
-		{
-			return FAppStyle::Get().GetBrush("DerivedData.RemoteCache.IdleBG");
-		}
-
-		case ERemoteCacheState::Busy :
-		{
-			return FAppStyle::Get().GetBrush("DerivedData.RemoteCache.BusyBG");
-		}	
-
-		case ERemoteCacheState::Unavailable:
-		{
-			return FAppStyle::Get().GetBrush("DerivedData.RemoteCache.UnavailableBG");
-		}
-
-		default:
-		case ERemoteCacheState::Warning:
-		{
-			return FAppStyle::Get().GetBrush("DerivedData.RemoteCache.WarningBG");
-		}
+		return IsBusy ? FAppStyle::Get().GetBrush("DerivedData.RemoteCache.BusyBG") : FAppStyle::Get().GetBrush("DerivedData.RemoteCache.IdleBG");
 	}
+
+	return FAppStyle::Get().GetBrush("DerivedData.RemoteCache.UnavailableBG");
 }
 
-const FSlateBrush* SZenStatusBarWidget::GetRemoteCacheStateBadgeIcon() const
+const FSlateBrush* SZenStatusBarWidget::GetServerStateBadgeIcon() const
 {
-	switch ( FDerivedDataInformation::GetRemoteCacheState())
+	if (IsRunning)
 	{
-		case ERemoteCacheState::Idle :
-		{
-			return FAppStyle::Get().GetBrush("DerivedData.RemoteCache.Idle");
-		}
-
-		case ERemoteCacheState::Busy :
-		{
-			return FAppStyle::Get().GetBrush("DerivedData.RemoteCache.Busy");
-		}	
-
-		case ERemoteCacheState::Unavailable:
-		{
-			return FAppStyle::Get().GetBrush("DerivedData.RemoteCache.Unavailable");
-		}
-
-		default:
-		case ERemoteCacheState::Warning:
-		{
-			return FAppStyle::Get().GetBrush("DerivedData.RemoteCache.Warning");
-		}
+		return IsBusy ? FAppStyle::Get().GetBrush("DerivedData.RemoteCache.Busy") : FAppStyle::Get().GetBrush("DerivedData.RemoteCache.Idle");
 	}
+
+	return FAppStyle::Get().GetBrush("DerivedData.RemoteCache.Unavailable");
 }
 
+
 #undef LOCTEXT_NAMESPACE
\ No newline at end of file
diff --git a/Engine/Source/Editor/ZenEditor/Private/SZenStatusBar.h b/Engine/Source/Editor/ZenEditor/Private/SZenStatusBar.h
--- a/Engine/Source/Editor/ZenEditor/Private/SZenStatusBar.h
+++ b/Engine/Source/Editor/ZenEditor/Private/SZenStatusBar.h
@@ -8,6 +8,7 @@
 #include "Types/SlateEnums.h"
 #include "Widgets/DeclarativeSyntaxSupport.h"
 #include "Widgets/SCompoundWidget.h"
+#include "Styling/SlateColor.h"
 
 class FUICommandInfo;
 class FUICommandList;
@@ -15,6 +16,8 @@ class SNotificationItem;
 class SWidget;
 struct FSlateBrush;
 
+namespace UE::Zen { class FServiceInstanceManager; }
+
 class FZenStausBarCommands : public TCommands<FZenStausBarCommands>
 {
 public:
@@ -25,21 +28,30 @@ public:
 
 private:
 
-	static void ChangeSettings_Clicked();
+	static void ChangeCacheSettings_Clicked();
 	static void ViewCacheStatistics_Clicked();
 	static void ViewResourceUsage_Clicked();
 	static void ViewServerStatus_Clicked();
 	static void LaunchZenDashboard_Clicked();
+	static void ViewZenStore_Clicked();
+	static void StartZenServer_Clicked();
+	static void StopZenServer_Clicked();
+	static void RestartZenServer_Clicked();
 
 public:
 
-	TSharedPtr< FUICommandInfo > ChangeSettings;
+	TSharedPtr< FUICommandInfo > ChangeCacheSettings;
 	TSharedPtr< FUICommandInfo > ViewResourceUsage;
 	TSharedPtr< FUICommandInfo > ViewCacheStatistics;
 	TSharedPtr< FUICommandInfo > ViewServerStatus;
 	TSharedPtr< FUICommandInfo > LaunchZenDashboard;
+	TSharedPtr< FUICommandInfo > ViewZenStore;
+	TSharedPtr< FUICommandInfo > StartZenServer;
+	TSharedPtr< FUICommandInfo > StopZenServer;
+	TSharedPtr< FUICommandInfo > RestartZenServer;
 
 	static TSharedRef<FUICommandList> ActionList;
+	static TSharedPtr<UE::Zen::FServiceInstanceManager> ServiceInstanceManager;
 };
 
 
@@ -53,10 +65,12 @@ class SZenStatusBarWidget : public SCompoundWidget
 private:
 
 	FText						GetTitleToolTipText() const;
-	FText						GetRemoteCacheToolTipText() const;
 	FText						GetTitleText() const;
-	const FSlateBrush*			GetRemoteCacheStateBackgroundIcon() const;
-	const FSlateBrush*			GetRemoteCacheStateBadgeIcon() const;
+
+	FText						GetServerStateToolTipText() const;
+	const FSlateBrush*			GetServerStateBadgeIcon() const;
+	const FSlateBrush*			GetServerStateBackgroundIcon() const;
+
 	TSharedRef<SWidget>			CreateStatusBarMenu();
 	EActiveTimerReturnType		UpdateBusyIndicator(double InCurrentTime, float InDeltaTime);
 	EActiveTimerReturnType		UpdateWarnings(double InCurrentTime, float InDeltaTime);
@@ -65,7 +79,13 @@ private:
 	double ElapsedUploadTime = 0;
 	double ElapsedBusyTime = 0;
 
-	bool bBusy = false;
+	bool IsRunning = false;
+	bool IsDownloading = false;
+	bool IsUploading = false;
+	bool IsReading = false;
+	bool IsWriting = false;
+	bool IsBusy = false;
 
 	TSharedPtr<SNotificationItem> NotificationItem;
+	TSharedPtr<UE::Zen::FServiceInstanceManager> ServiceInstanceManager;
 };
diff --git a/Engine/Source/Editor/ZenEditor/Private/ZenEditor.cpp b/Engine/Source/Editor/ZenEditor/Private/ZenEditor.cpp
--- a/Engine/Source/Editor/ZenEditor/Private/ZenEditor.cpp
+++ b/Engine/Source/Editor/ZenEditor/Private/ZenEditor.cpp
@@ -161,4 +161,28 @@ TSharedRef<SDockTab> FZenEditor::CreateZenServerStatusTab(const FSpawnTabArgs& A
 		];
 }
 
+void FZenEditor::StartZenServer()
+{
+	UE::Zen::FZenLocalServiceRunContext RunContext;
+	if (UE::Zen::TryGetLocalServiceRunContext(RunContext))
+	{
+		UE::Zen::StartLocalService(RunContext);
+	}
+}
+
+void FZenEditor::StopZenServer()
+{
+	UE::Zen::FZenLocalServiceRunContext RunContext;
+	if (UE::Zen::TryGetLocalServiceRunContext(RunContext))
+	{
+		UE::Zen::StopLocalService(*RunContext.GetDataPath());
+	}
+}
+
+void FZenEditor::RestartZenServer()
+{
+	StopZenServer();
+	StartZenServer();
+}
+
 #undef LOCTEXT_NAMESPACE
diff --git a/Engine/Source/Editor/ZenEditor/Public/ZenEditor.h b/Engine/Source/Editor/ZenEditor/Public/ZenEditor.h
--- a/Engine/Source/Editor/ZenEditor/Public/ZenEditor.h
+++ b/Engine/Source/Editor/ZenEditor/Public/ZenEditor.h
@@ -39,6 +39,9 @@ public:
 	void ShowResourceUsageTab();
 	void ShowCacheStatisticsTab();
 	void ShowZenServerStatusTab();
+	void StartZenServer();
+	void StopZenServer();
+	void RestartZenServer();
 
 private:
 
