Zen Tool Tray UX Icon Refesh

diff --git a/Engine/Source/Editor/EditorStyle/Private/StarshipStyle.cpp b/Engine/Source/Editor/EditorStyle/Private/StarshipStyle.cpp
--- a/Engine/Source/Editor/EditorStyle/Private/StarshipStyle.cpp
+++ b/Engine/Source/Editor/EditorStyle/Private/StarshipStyle.cpp
@@ -1097,17 +1097,31 @@ void FStarshipEditorStyle::FStyle::SetupGeneralStyles()
 
 	// Zen
 	{
-		Set("Zen.Icons.Server", new CORE_IMAGE_BRUSH_SVG("Starship/Common/Server", Icon16x16));
-		Set("Zen.Icons.Server.Start", new CORE_IMAGE_BRUSH_SVG("Starship/Common/Play", Icon16x16));
-		Set("Zen.Icons.Server.Stop", new CORE_IMAGE_BRUSH_SVG("Starship/Common/Stop", Icon16x16));
-		Set("Zen.Icons.Server.Restart", new CORE_IMAGE_BRUSH_SVG("Starship/Common/Update", Icon16x16));
-		Set("Zen.Icons.LaunchDashboard", new CORE_IMAGE_BRUSH_SVG("Starship/Common/Monitor", Icon16x16));
-		Set("Zen.Icons.ImportSnapshot", new CORE_IMAGE_BRUSH_SVG("Starship/Common/Import", Icon16x16));
-		Set("Zen.Icons.Store", new CORE_IMAGE_BRUSH_SVG("Starship/Common/Cylinder", Icon16x16));
-		Set("Zen.Icons.WebBrowser", new IMAGE_BRUSH_SVG("Starship/Common/WebBrowser", Icon12x12));
-		Set("Zen.Icons.FolderExplore", new IMAGE_BRUSH_SVG("Starship/ContentBrowser/show-in-explorer", Icon12x12));
-		Set("Zen.Icons.Clipboard", new IMAGE_BRUSH_SVG("Starship/Common/Clipboard", Icon16x16));	
-
+		Set("Zen.Server", new CORE_IMAGE_BRUSH_SVG("Starship/Common/Server", Icon16x16));
+		Set("Zen.Server.Start", new CORE_IMAGE_BRUSH_SVG("Starship/Common/Play", Icon16x16));
+		Set("Zen.Server.Stop", new CORE_IMAGE_BRUSH_SVG("Starship/Common/Stop", Icon16x16));
+		Set("Zen.Server.Restart", new CORE_IMAGE_BRUSH_SVG("Starship/Common/Update", Icon16x16));
+		Set("Zen.Store", new CORE_IMAGE_BRUSH_SVG("Starship/Common/Cylinder", Icon16x16));
+		Set("Zen.WebBrowser", new IMAGE_BRUSH_SVG("Starship/Common/WebBrowser", Icon12x12));
+		Set("Zen.FolderExplore", new IMAGE_BRUSH_SVG("Starship/ContentBrowser/show-in-explorer", Icon12x12));
+		Set("Zen.Clipboard", new IMAGE_BRUSH_SVG("Starship/Common/Clipboard", Icon16x16));	
+		Set("Zen.Dashboard.Icon", new IMAGE_BRUSH_SVG("Starship/Zen/Zen_16", Icon16x16));
+		Set("Zen.Dashboard.Application", new IMAGE_BRUSH_SVG("Starship/Zen/Zen_24", Icon24x24));
+
+		Set("Zen.Status.UnavailableBG", new IMAGE_BRUSH_SVG("Starship/Zen/ZenStatusWithBadge_16", Icon16x16));
+		Set("Zen.Status.Unavailable", new IMAGE_BRUSH_SVG("Starship/Zen/ZenStatusBadge", Icon16x16));
+
+		Set("Zen.Status.IdleBG", new IMAGE_BRUSH_SVG("Starship/DerivedData/DD_RemoteCache_IdleBG", Icon16x16));
+		Set("Zen.Status.Idle", new IMAGE_BRUSH_SVG("Starship/DerivedData/DD_RemoteCache_Idle", Icon16x16, EStyleColor::Success));
+
+		Set("Zen.Status.BusyBG", new IMAGE_BRUSH_SVG("Starship/DerivedData/DD_RemoteCache_UpDownBG", Icon16x16));
+		Set("Zen.Status.Busy", new IMAGE_BRUSH_SVG("Starship/DerivedData/DD_RemoteCache_UpDown", Icon16x16, FLinearColor(0.3f, 0.3f, 0.3f, 1.0f)));
+
+		Set("Zen.Status.Uploading", new IMAGE_BRUSH_SVG("Starship/DerivedData/DD_Upload", Icon16x16, FLinearColor::Green));
+		Set("Zen.Status.Downloading", new IMAGE_BRUSH_SVG("Starship/DerivedData/DD_Download", Icon16x16, FLinearColor(0.0f, 1.0f, 1.0f, 1.0f)));
+		
+		Set("Zen.Status.WarningBG", new IMAGE_BRUSH_SVG("Starship/DerivedData/DD_RemoteCache_WarningBG", Icon16x16));
+		Set("Zen.Status.Warning", new IMAGE_BRUSH_SVG("Starship/DerivedData/DD_RemoteCache_Warning", Icon16x16, EStyleColor::Warning));
 	}
 
 	// Derived Data
@@ -1115,9 +1129,7 @@ void FStarshipEditorStyle::FStyle::SetupGeneralStyles()
 		Set("DerivedData.Cache.Settings", new IMAGE_BRUSH_SVG("Starship/DerivedData/DD_Cache_Settings", Icon16x16));
 		Set("DerivedData.Cache.Statistics", new IMAGE_BRUSH_SVG("Starship/DerivedData/DD_Cache_Statistics", Icon16x16));
 		Set("DerivedData.ResourceUsage", new IMAGE_BRUSH_SVG("Starship/DerivedData/DD_Resource_Usage", Icon16x16));
-		Set("DerivedData.RemoteCache.Uploading", new IMAGE_BRUSH_SVG("Starship/DerivedData/DD_Upload", Icon16x16, EStyleColor::AccentGreen));
-		Set("DerivedData.RemoteCache.Downloading", new IMAGE_BRUSH_SVG("Starship/DerivedData/DD_Download", Icon16x16, EStyleColor::AccentBlue));
-
+		
 		Set("DerivedData.RemoteCache.Uploading", new IMAGE_BRUSH_SVG("Starship/DerivedData/DD_Upload", Icon16x16, FLinearColor::Green));
 		Set("DerivedData.RemoteCache.Downloading", new IMAGE_BRUSH_SVG("Starship/DerivedData/DD_Download", Icon16x16, FLinearColor(0.0f, 1.0f, 1.0f, 1.0f)));
 
diff --git a/Engine/Source/Editor/StorageServerWidgets/Private/SZenServiceStatus.cpp b/Engine/Source/Editor/StorageServerWidgets/Private/SZenServiceStatus.cpp
--- a/Engine/Source/Editor/StorageServerWidgets/Private/SZenServiceStatus.cpp
+++ b/Engine/Source/Editor/StorageServerWidgets/Private/SZenServiceStatus.cpp
@@ -257,7 +257,7 @@ TSharedRef<SWidget> SZenServiceStatus::GetGridPanel()
 			})
 			[
 				SNew(SImage)
-				.Image(FAppStyle::Get().GetBrush("Zen.Icons.WebBrowser"))
+				.Image(FAppStyle::Get().GetBrush("Zen.WebBrowser"))
 				.ColorAndOpacity(FSlateColor::UseForeground())
 			]
 		]
@@ -355,7 +355,7 @@ TSharedRef<SWidget> SZenServiceStatus::GetGridPanel()
 			.OnClicked(this, &SZenServiceStatus::ExploreDataPath_OnClicked)
 			[
 				SNew(SImage)
-				.Image(FAppStyle::Get().GetBrush("Zen.Icons.FolderExplore"))
+				.Image(FAppStyle::Get().GetBrush("Zen.FolderExplore"))
 				.ColorAndOpacity(FSlateColor::UseForeground())
 			]
 		]
@@ -408,7 +408,7 @@ TSharedRef<SWidget> SZenServiceStatus::GetGridPanel()
 			.OnClicked(this, &SZenServiceStatus::ExploreExecutablePath_OnClicked)
 			[
 				SNew(SImage)
-				.Image(FAppStyle::Get().GetBrush("Zen.Icons.FolderExplore"))
+				.Image(FAppStyle::Get().GetBrush("Zen.FolderExplore"))
 				.ColorAndOpacity(FSlateColor::UseForeground())
 			]
 		]
diff --git a/Engine/Source/Editor/ZenEditor/Private/SZenStatusBar.cpp b/Engine/Source/Editor/ZenEditor/Private/SZenStatusBar.cpp
--- a/Engine/Source/Editor/ZenEditor/Private/SZenStatusBar.cpp
+++ b/Engine/Source/Editor/ZenEditor/Private/SZenStatusBar.cpp
@@ -85,8 +85,8 @@ void FZenStausBarCommands::RegisterCommands()
 	UI_COMMAND(ViewCacheStatistics, "Cache Statistics", "Opens the Cache Statistics panel.", EUserInterfaceActionType::Button, FInputChord());
 	UI_COMMAND(ViewResourceUsage, "Cache Resource Usage", "Opens the Cache Resource Usage panel.", EUserInterfaceActionType::Button, FInputChord());
 	UI_COMMAND(ViewServerStatus, "Server Status", "Opens the Zen Server Status panel.", EUserInterfaceActionType::Button, FInputChord());
-	UI_COMMAND(LaunchZenDashboard, "Launch Dashboard", "Launches the Zen Dashboard utility.", EUserInterfaceActionType::Button, FInputChord());
 	UI_COMMAND(ViewZenStore, "View Store", "Opens the Zen Store page in your browser.", EUserInterfaceActionType::Button, FInputChord());
+	UI_COMMAND(LaunchZenDashboard, "Launch Dashboard", "Launches the Zen Dashboard utility.", EUserInterfaceActionType::Button, FInputChord());
 	UI_COMMAND(StartZenServer, "Start Server", "Starts the Zen Server.", EUserInterfaceActionType::Button, FInputChord());
 	UI_COMMAND(StopZenServer, "Stop Server", "Stops the Zen Server.", EUserInterfaceActionType::Button, FInputChord());
 	UI_COMMAND(RestartZenServer, "Restart Server", "Restarts the Zen Server.", EUserInterfaceActionType::Button, FInputChord());
@@ -112,13 +112,13 @@ void FZenStausBarCommands::RegisterCommands()
 	);
 
 	ActionList->MapAction(
-		LaunchZenDashboard,
-		FExecuteAction::CreateStatic(&FZenStausBarCommands::LaunchZenDashboard_Clicked)
+		ViewZenStore,
+		FExecuteAction::CreateStatic(&FZenStausBarCommands::ViewZenStore_Clicked)
 	);
 
 	ActionList->MapAction(
-		ViewZenStore,
-		FExecuteAction::CreateStatic(&FZenStausBarCommands::ViewZenStore_Clicked)
+		LaunchZenDashboard,
+		FExecuteAction::CreateStatic(&FZenStausBarCommands::LaunchZenDashboard_Clicked)
 	);
 
 	ActionList->MapAction(
@@ -160,12 +160,6 @@ void FZenStausBarCommands::ViewServerStatus_Clicked()
 	ZenEditor.ShowZenServerStatusTab();
 }
 
-void FZenStausBarCommands::LaunchZenDashboard_Clicked()
-{
-	using namespace UE::Zen;
-	FZenDashboardLauncher::Get()->StartZenDashboard(FZenDashboardLauncher::Get()->GetZenDashboardApplicationPath());
-}
-
 void FZenStausBarCommands::StartZenServer_Clicked()
 {
 	FZenEditor& ZenEditor = FModuleManager::LoadModuleChecked<FZenEditor>("ZenEditor");
@@ -196,6 +190,13 @@ void FZenStausBarCommands::ViewZenStore_Clicked()
 	}
 }
 
+void FZenStausBarCommands::LaunchZenDashboard_Clicked()
+{
+	using namespace UE::Zen;
+	FZenDashboardLauncher::Get()->StartZenDashboard(FZenDashboardLauncher::Get()->GetZenDashboardApplicationPath());
+}
+
+
 TSharedRef<SWidget> SZenStatusBarWidget::CreateStatusBarMenu()
 {
 	UToolMenu* Menu = UToolMenus::Get()->RegisterMenu(ZenMenuName, NAME_None, EMultiBoxType::Menu, false);
@@ -207,9 +208,10 @@ TSharedRef<SWidget> SZenStatusBarWidget::CreateStatusBarMenu()
 			FZenStausBarCommands::Get().LaunchZenDashboard,
 			TAttribute<FText>(),
 			TAttribute<FText>(),
-			FSlateIcon(FAppStyle::GetAppStyleSetName(), "Zen.Icons.LaunchDashboard")
+			FSlateIcon(FAppStyle::GetAppStyleSetName(), "Zen.Dashboard.Icon")
 		);
 	}
+
 	{
 		FToolMenuSection& Section = Menu->AddSection(ZenMenuServerSectionName, LOCTEXT("ZenMenuServerSection", "Server"));
 
@@ -217,7 +219,7 @@ TSharedRef<SWidget> SZenStatusBarWidget::CreateStatusBarMenu()
 			FZenStausBarCommands::Get().ViewServerStatus,
 			TAttribute<FText>(),
 			TAttribute<FText>(),
-			FSlateIcon(FAppStyle::GetAppStyleSetName(), "Zen.Icons.Server")
+			FSlateIcon(FAppStyle::GetAppStyleSetName(), "Zen.Server")
 		);
 
 		{
@@ -234,14 +236,14 @@ TSharedRef<SWidget> SZenStatusBarWidget::CreateStatusBarMenu()
 								FZenStausBarCommands::Get().StopZenServer,
 								TAttribute<FText>(),
 								TAttribute<FText>(),
-								FSlateIcon(FAppStyle::GetAppStyleSetName(), "Zen.Icons.Server.Stop")
+								FSlateIcon(FAppStyle::GetAppStyleSetName(), "Zen.Server.Stop")
 							);
 
 							Section.AddMenuEntry(
 								FZenStausBarCommands::Get().RestartZenServer,
 								TAttribute<FText>(),
 								TAttribute<FText>(),
-								FSlateIcon(FAppStyle::GetAppStyleSetName(), "Zen.Icons.Server.Restart")
+								FSlateIcon(FAppStyle::GetAppStyleSetName(), "Zen.Server.Restart")
 							);
 						}
 						else
@@ -250,7 +252,7 @@ TSharedRef<SWidget> SZenStatusBarWidget::CreateStatusBarMenu()
 								FZenStausBarCommands::Get().StartZenServer,
 								TAttribute<FText>(),
 								TAttribute<FText>(),
-								FSlateIcon(FAppStyle::GetAppStyleSetName(), "Zen.Icons.Server.Start")
+								FSlateIcon(FAppStyle::GetAppStyleSetName(), "Zen.Server.Start")
 							);
 						}
 					}));
@@ -264,7 +266,7 @@ TSharedRef<SWidget> SZenStatusBarWidget::CreateStatusBarMenu()
 			FZenStausBarCommands::Get().ViewZenStore,
 			TAttribute<FText>(),
 			TAttribute<FText>(),
-			FSlateIcon(FAppStyle::GetAppStyleSetName(), "Zen.Icons.Store")
+			FSlateIcon(FAppStyle::GetAppStyleSetName(), "Zen.Store")
 		);
 	}
 
@@ -342,7 +344,7 @@ void SZenStatusBarWidget::Construct(const FArguments& InArgs)
 				.VAlign(VAlign_Top)
 				[
 					SNew(SImage)
-					.Image(FAppStyle::Get().GetBrush("DerivedData.RemoteCache.Uploading"))
+					.Image(FAppStyle::Get().GetBrush("Zen.Status.Uploading"))
 					.ColorAndOpacity_Lambda([this] { return IsUploading? FLinearColor::White.CopyWithNewOpacity(FMath::MakePulsatingValue(ElapsedUploadTime, 2)) : FLinearColor(0,0,0,0); })
 					.ToolTipText_Lambda([this] { return GetServerStateToolTipText(); })
 				]
@@ -352,7 +354,7 @@ void SZenStatusBarWidget::Construct(const FArguments& InArgs)
 				.VAlign(VAlign_Top)
 				[
 					SNew(SImage)
-					.Image(FAppStyle::Get().GetBrush("DerivedData.RemoteCache.Downloading"))
+					.Image(FAppStyle::Get().GetBrush("Zen.Status.Downloading"))
 					.ColorAndOpacity_Lambda([this] { return IsDownloading ? FLinearColor::White.CopyWithNewOpacity(FMath::MakePulsatingValue(ElapsedDownloadTime, 2)) : FLinearColor(0, 0, 0, 0); })
 					.ToolTipText_Lambda([this] { return GetServerStateToolTipText(); })
 				]				
@@ -542,20 +544,20 @@ const FSlateBrush* SZenStatusBarWidget::GetServerStateBackgroundIcon() const
 {
 	if (IsRunning)
 	{
-		return IsBusy ? FAppStyle::Get().GetBrush("DerivedData.RemoteCache.BusyBG") : FAppStyle::Get().GetBrush("DerivedData.RemoteCache.IdleBG");
+		return IsBusy ? FAppStyle::Get().GetBrush("Zen.Status.BusyBG") : FAppStyle::Get().GetBrush("Zen.Status.IdleBG");
 	}
 
-	return FAppStyle::Get().GetBrush("DerivedData.RemoteCache.UnavailableBG");
+	return FAppStyle::Get().GetBrush("Zen.Status.UnavailableBG");
 }
 
 const FSlateBrush* SZenStatusBarWidget::GetServerStateBadgeIcon() const
 {
 	if (IsRunning)
 	{
-		return IsBusy ? FAppStyle::Get().GetBrush("DerivedData.RemoteCache.Busy") : FAppStyle::Get().GetBrush("DerivedData.RemoteCache.Idle");
+		return IsBusy ? FAppStyle::Get().GetBrush("Zen.Status.Busy") : FAppStyle::Get().GetBrush("Zen.Status.Idle");
 	}
 
-	return FAppStyle::Get().GetBrush("DerivedData.RemoteCache.Unavailable");
+	return FAppStyle::Get().GetBrush("Zen.Status.Unavailable");
 }
 
 
diff --git a/Engine/Source/Editor/ZenEditor/Private/ZenEditor.cpp b/Engine/Source/Editor/ZenEditor/Private/ZenEditor.cpp
--- a/Engine/Source/Editor/ZenEditor/Private/ZenEditor.cpp
+++ b/Engine/Source/Editor/ZenEditor/Private/ZenEditor.cpp
@@ -42,7 +42,7 @@ void FZenEditor::StartupModule()
 		.SetGroup(WorkspaceMenu::GetMenuStructure().GetToolsCategory())
 		.SetIcon(CacheStatisticsIcon);
 
-	const FSlateIcon ZenServerIcon(FAppStyle::GetAppStyleSetName(), "Zen.Icons.Server");
+	const FSlateIcon ZenServerIcon(FAppStyle::GetAppStyleSetName(), "Zen.Server");
 
 	FGlobalTabmanager::Get()->RegisterNomadTabSpawner(ZenServerStatusTabName, FOnSpawnTab::CreateRaw(this, &FZenEditor::CreateZenServerStatusTab))
 		.SetDisplayName(LOCTEXT("ZenServerStatusTabTitle", "Zen Server Status"))
diff --git a/Engine/Source/Programs/BuildStorageTool/Private/BuildStorageToolStyle.cpp b/Engine/Source/Programs/BuildStorageTool/Private/BuildStorageToolStyle.cpp
--- a/Engine/Source/Programs/BuildStorageTool/Private/BuildStorageToolStyle.cpp
+++ b/Engine/Source/Programs/BuildStorageTool/Private/BuildStorageToolStyle.cpp
@@ -56,8 +56,8 @@ TSharedRef< FSlateStyleSet > FBuildStorageToolStyle::Create()
 
 	Style.Set( "AppIcon", new IMAGE_BRUSH( "Icons/DefaultAppIcon", FVector2D(20,20)) );
 
-	Style.Set("Zen.Icons.FolderExplore", new EDITOR_IMAGE_BRUSH_SVG("Starship/Common/ContentBrowser", Icon12x12));
-	Style.Set("Zen.Icons.Clipboard", new EDITOR_IMAGE_BRUSH_SVG("Starship/Common/Clipboard", Icon16x16));
+	Style.Set("Zen.FolderExplore", new EDITOR_IMAGE_BRUSH_SVG("Starship/Common/ContentBrowser", Icon12x12));
+	Style.Set("Zen.Clipboard", new EDITOR_IMAGE_BRUSH_SVG("Starship/Common/Clipboard", Icon16x16));
 
 	const FTextBlockStyle DefaultText = FTextBlockStyle()
 		.SetFont(DEFAULT_FONT("Bold", 10))
diff --git a/Engine/Source/Programs/BuildStorageTool/Private/SMessageDialog.cpp b/Engine/Source/Programs/BuildStorageTool/Private/SMessageDialog.cpp
--- a/Engine/Source/Programs/BuildStorageTool/Private/SMessageDialog.cpp
+++ b/Engine/Source/Programs/BuildStorageTool/Private/SMessageDialog.cpp
@@ -124,7 +124,7 @@ public:
 						.Content()
 						[
 							SNew(SImage)
-							.Image(FAppStyle::Get().GetBrush("Zen.Icons.Clipboard"))
+							.Image(FAppStyle::Get().GetBrush("Zen.Clipboard"))
 							.ColorAndOpacity(FSlateColor::UseForeground())
 						]
 					]
diff --git a/Engine/Source/Programs/ZenDashboard/Private/SMessageDialog.cpp b/Engine/Source/Programs/ZenDashboard/Private/SMessageDialog.cpp
--- a/Engine/Source/Programs/ZenDashboard/Private/SMessageDialog.cpp
+++ b/Engine/Source/Programs/ZenDashboard/Private/SMessageDialog.cpp
@@ -124,7 +124,7 @@ public:
 						.Content()
 						[
 							SNew(SImage)
-							.Image(FAppStyle::Get().GetBrush("Zen.Icons.Clipboard"))
+							.Image(FAppStyle::Get().GetBrush("Zen.Clipboard"))
 							.ColorAndOpacity(FSlateColor::UseForeground())
 						]
 					]
diff --git a/Engine/Source/Programs/ZenDashboard/Private/ZenDashboardStyle.cpp b/Engine/Source/Programs/ZenDashboard/Private/ZenDashboardStyle.cpp
--- a/Engine/Source/Programs/ZenDashboard/Private/ZenDashboardStyle.cpp
+++ b/Engine/Source/Programs/ZenDashboard/Private/ZenDashboardStyle.cpp
@@ -59,12 +59,13 @@ TSharedRef< FSlateStyleSet > FZenDashboardStyle::Create()
 	const FVector2D Icon10x10(10.f, 10.f);
 	const FVector2D Icon12x12(12.0f, 12.0f);
 	const FVector2D Icon16x16(16.0f, 16.0f);
+	const FVector2D Icon20x20(20.0f, 20.0f);
 
-	Style.Set( "AppIcon", new IMAGE_BRUSH( "Icons/DefaultAppIcon", FVector2D(20,20)) );
+	Style.Set( "AppIcon", new EDITOR_IMAGE_BRUSH_SVG( "Starship/Zen/Zen_24", Icon20x20) );
 
-	Style.Set("Zen.Icons.WebBrowser", new EDITOR_IMAGE_BRUSH_SVG("Starship/Common/WebBrowser", Icon12x12));
-	Style.Set("Zen.Icons.FolderExplore", new EDITOR_IMAGE_BRUSH_SVG("Starship/Common/ContentBrowser", Icon12x12));
-	Style.Set("Zen.Icons.Clipboard", new EDITOR_IMAGE_BRUSH_SVG("Starship/Common/Clipboard", Icon16x16));
+	Style.Set("Zen.WebBrowser", new EDITOR_IMAGE_BRUSH_SVG("Starship/Common/WebBrowser", Icon12x12));
+	Style.Set("Zen.FolderExplore", new EDITOR_IMAGE_BRUSH_SVG("Starship/Common/ContentBrowser", Icon12x12));
+	Style.Set("Zen.Clipboard", new EDITOR_IMAGE_BRUSH_SVG("Starship/Common/Clipboard", Icon16x16));
 
 	Style.Set("Icons.ChevronRight", new CORE_IMAGE_BRUSH_SVG("Starship/Common/chevron-right", Icon16x16));
 	Style.Set("Icons.ChevronDown", new CORE_IMAGE_BRUSH_SVG("Starship/Common/chevron-down", Icon16x16));
